Copying fixes for Vehicle Templates and Ocean build breaks, and QA material break

diff --git a/Engine/Plugins/Runtime/OculusInput/Source/OculusInput.Build.cs b/Engine/Plugins/Runtime/OculusInput/Source/OculusInput.Build.cs
--- a/Engine/Plugins/Runtime/OculusInput/Source/OculusInput.Build.cs
+++ b/Engine/Plugins/Runtime/OculusInput/Source/OculusInput.Build.cs
@@ -30,8 +30,7 @@ namespace UnrealBuildTool.Rules
             // Currently, the Rift is only supported on windows and mac platforms
             if (Target.Platform == UnrealTargetPlatform.Win32 || Target.Platform == UnrealTargetPlatform.Win64 || Target.Platform == UnrealTargetPlatform.Mac)
             {
-//				Definitions.Add("USE_OVR_MOTION_SDK=" + ( UEBuildConfiguration.bOculusInputControllerSupport ? "1" : "0" ) );
-				Definitions.Add("USE_OVR_MOTION_SDK=" + ( true ? "1" : "0" ) );
+				Definitions.Add("USE_OVR_MOTION_SDK=1");
 				
 				PrivateDependencyModuleNames.AddRange(new string[] { "LibOVR" });
             }
diff --git a/Templates/TP_Vehicle/Source/TP_Vehicle/TP_VehiclePawn.cpp b/Templates/TP_Vehicle/Source/TP_Vehicle/TP_VehiclePawn.cpp
--- a/Templates/TP_Vehicle/Source/TP_Vehicle/TP_VehiclePawn.cpp
+++ b/Templates/TP_Vehicle/Source/TP_Vehicle/TP_VehiclePawn.cpp
@@ -159,12 +159,6 @@ void ATP_VehiclePawn::EnableIncarView(const bool bState, const bool bForce)
 			OnResetVR();
 			Camera->Deactivate();
 			InternalCamera->Activate();
-			
-			APlayerController* PlayerController = Cast<APlayerController>(GetController());
-			if ( (PlayerController != nullptr) && (PlayerController->PlayerCameraManager != nullptr ) )
-			{
-				PlayerController->PlayerCameraManager->bFollowHmdOrientation = true;
-			}
 		}
 		else
 		{
diff --git a/Templates/TP_VehicleAdv/Source/TP_VehicleAdv/TP_VehicleAdvPawn.cpp b/Templates/TP_VehicleAdv/Source/TP_VehicleAdv/TP_VehicleAdvPawn.cpp
--- a/Templates/TP_VehicleAdv/Source/TP_VehicleAdv/TP_VehicleAdvPawn.cpp
+++ b/Templates/TP_VehicleAdv/Source/TP_VehicleAdv/TP_VehicleAdvPawn.cpp
@@ -221,12 +221,6 @@ void ATP_VehicleAdvPawn::EnableIncarView(const bool bState)
 			OnResetVR();
 			Camera->Deactivate();
 			InternalCamera->Activate();
-			
-			APlayerController* PlayerController = GetWorld()->GetFirstPlayerController();
-			if ( (PlayerController != nullptr) && (PlayerController->PlayerCameraManager != nullptr ) )
-			{
-				PlayerController->PlayerCameraManager->bFollowHmdOrientation = true;
-			}
 		}
 		else
 		{
